[SequenceNavigator] Move Level Sequence Navigator to LevelSequenceNavigatorBridge plugin

diff --git a/Engine/Plugins/Experimental/LevelSequenceNavigatorBridge/LevelSequenceNavigatorBridge.uplugin b/Engine/Plugins/Experimental/LevelSequenceNavigatorBridge/LevelSequenceNavigatorBridge.uplugin
new file mode 100644
--- /dev/null
+++ b/Engine/Plugins/Experimental/LevelSequenceNavigatorBridge/LevelSequenceNavigatorBridge.uplugin
@@ -0,0 +1,36 @@
+﻿{
+	"FileVersion": 3,
+	"Version": 1,
+	"VersionName": "1.0",
+	"FriendlyName": "Level Sequence Navigator Bridge",
+	"Description": "Sequence Navigator Bridge for Level Sequences",
+	"Category": "Animation",
+	"CreatedBy": "Epic Games, Inc.",
+	"CreatedByURL": "https://epicgames.com",
+	"DocsURL": "",
+	"MarketplaceURL": "",
+	"CanContainContent": false,
+	"IsBetaVersion": false,
+	"IsExperimentalVersion": true,
+	"Installed": false,
+	"EnabledByDefault": true,
+	"Modules": [
+		{
+			"Name": "LevelSequenceNavigatorBridge",
+			"Type": "Editor",
+			"LoadingPhase": "PostDefault"
+		}
+	],
+	"Plugins": [
+		{
+			"Name": "LevelSequenceEditor",
+			"Enabled": true,
+			"Optional": true
+		},
+		{
+			"Name": "SequenceNavigator",
+			"Enabled": true,
+			"Optional": true
+		}
+	]
+}
diff --git a/Engine/Plugins/Experimental/LevelSequenceNavigatorBridge/Source/LevelSequenceNavigatorBridge/LevelSequenceNavigatorBridge.Build.cs b/Engine/Plugins/Experimental/LevelSequenceNavigatorBridge/Source/LevelSequenceNavigatorBridge/LevelSequenceNavigatorBridge.Build.cs
new file mode 100644
--- /dev/null
+++ b/Engine/Plugins/Experimental/LevelSequenceNavigatorBridge/Source/LevelSequenceNavigatorBridge/LevelSequenceNavigatorBridge.Build.cs
@@ -0,0 +1,38 @@
+﻿// Copyright Epic Games, Inc. All Rights Reserved.
+
+using UnrealBuildTool;
+
+public class LevelSequenceNavigatorBridge : ModuleRules
+{
+	public LevelSequenceNavigatorBridge(ReadOnlyTargetRules Target) : base(Target)
+	{
+		PCHUsage = PCHUsageMode.UseExplicitOrSharedPCHs;
+
+		PublicIncludePaths.AddRange(new string[] {
+
+		});
+
+		PrivateIncludePaths.AddRange(new string[] {
+
+		});
+
+		PublicDependencyModuleNames.AddRange(new string[] {
+			"Core"
+		});
+
+		PrivateDependencyModuleNames.AddRange(new string[] {
+			"CoreUObject",
+			"Engine",
+			"LevelSequence",
+			"MovieScene",
+			"Sequencer",
+			"SequenceNavigator",
+			"Slate",
+			"SlateCore"
+		});
+
+		DynamicallyLoadedModuleNames.AddRange(new string[] {
+
+		});
+	}
+}
diff --git a/Engine/Plugins/Experimental/LevelSequenceNavigatorBridge/Source/LevelSequenceNavigatorBridge/Private/LevelSequenceNavigationToolProvider.cpp b/Engine/Plugins/Experimental/LevelSequenceNavigatorBridge/Source/LevelSequenceNavigatorBridge/Private/LevelSequenceNavigationToolProvider.cpp
new file mode 100644
--- /dev/null
+++ b/Engine/Plugins/Experimental/LevelSequenceNavigatorBridge/Source/LevelSequenceNavigatorBridge/Private/LevelSequenceNavigationToolProvider.cpp
@@ -0,0 +1,219 @@
+﻿// Copyright Epic Games, Inc. All Rights Reserved.
+
+#include "LevelSequenceNavigationToolProvider.h"
+#include "Columns/NavigationToolColorColumn.h"
+#include "Columns/NavigationToolColumnExtender.h"
+#include "Columns/NavigationToolCommentColumn.h"
+#include "Columns/NavigationToolHBiasColumn.h"
+#include "Columns/NavigationToolInTimeColumn.h"
+#include "Columns/NavigationToolItemsColumn.h"
+#include "Columns/NavigationToolLabelColumn.h"
+#include "Columns/NavigationToolLengthColumn.h"
+#include "Columns/NavigationToolLockColumn.h"
+#include "Columns/NavigationToolMarkerVisibilityColumn.h"
+#include "Columns/NavigationToolOutTimeColumn.h"
+#include "Columns/NavigationToolPlayheadColumn.h"
+#include "Columns/NavigationToolRevisionControlColumn.h"
+#include "Columns/NavigationToolStartFrameOffsetColumn.h"
+#include "Columns/NavigationToolTakeColumn.h"
+#include "Columns/NavigationToolDeactiveStateColumn.h"
+#include "Filters/Filters/NavigationToolBuiltInFilterParams.h"
+#include "INavigationTool.h"
+#include "ISequencer.h"
+#include "ISourceControlModule.h"
+#include "Items/NavigationToolSequence.h"
+#include "LevelSequence.h"
+#include "MovieScene.h"
+#include "MovieSceneSequence.h"
+#include "NavigationToolSettings.h"
+
+#define LOCTEXT_NAMESPACE "LevelSequenceNavigationToolProvider"
+
+const FName FLevelSequenceNavigationToolProvider::Identifier = TEXT("LevelSequence");
+
+const FText FLevelSequenceNavigationToolProvider::AnimationColumnViewName = LOCTEXT("AnimationColumnViewName", "Animation");
+
+FLevelSequenceNavigationToolProvider::FLevelSequenceNavigationToolProvider()
+	: ToolCommands(MakeShared<FUICommandList>())
+{
+}
+
+FName FLevelSequenceNavigationToolProvider::GetIdentifier() const
+{
+	return Identifier;
+}
+
+TSet<TSubclassOf<UMovieSceneSequence>> FLevelSequenceNavigationToolProvider::GetSupportedSequenceClasses() const
+{
+	return { ULevelSequence::StaticClass() };
+}
+
+FText FLevelSequenceNavigationToolProvider::GetDefaultColumnView() const
+{
+	return AnimationColumnViewName;
+}
+
+void FLevelSequenceNavigationToolProvider::BindCommands(const TSharedRef<FUICommandList>& InCommandList)
+{
+	InCommandList->Append(ToolCommands);
+}
+
+void FLevelSequenceNavigationToolProvider::OnActivate()
+{
+	// @TODO: Move from core Navigation Tool
+	/*ExtendToolToolBar();
+	ExtendToolItemContextMenu();
+
+	DragDropInitializedDelegate = FNavigationToolExtender::OnItemDragDropOpInitialized()
+		.AddSP(this, &FAvaNavigationToolProvider::ExtendToolItemDragDropOp);*/
+}
+
+void FLevelSequenceNavigationToolProvider::OnDeactivate()
+{
+	// @TODO: Move from core Navigation Tool
+	/*RemoveToolToolBarExtension();
+	RemoveToolItemContextMenuExtension();
+
+	FNavigationToolExtender::OnItemDragDropOpInitialized()
+		.Remove(DragDropInitializedDelegate);*/
+}
+
+void FLevelSequenceNavigationToolProvider::OnExtendColumns(UE::SequenceNavigator::FNavigationToolColumnExtender& OutExtender)
+{
+	using namespace UE::SequenceNavigator;
+
+	OutExtender.AddColumn<FNavigationToolPlayheadColumn>();
+	OutExtender.AddColumn<FNavigationToolDeactiveStateColumn>();
+	OutExtender.AddColumn<FNavigationToolMarkerVisibilityColumn>();
+	OutExtender.AddColumn<FNavigationToolLockColumn>();
+	OutExtender.AddColumn<FNavigationToolColorColumn>();
+	OutExtender.AddColumn<FNavigationToolLabelColumn>();
+	OutExtender.AddColumn<FNavigationToolItemsColumn>();
+	OutExtender.AddColumn<FNavigationToolInTimeColumn>();
+	OutExtender.AddColumn<FNavigationToolOutTimeColumn>();
+	OutExtender.AddColumn<FNavigationToolLengthColumn>();
+	OutExtender.AddColumn<FNavigationToolHBiasColumn>();
+	OutExtender.AddColumn<FNavigationToolStartFrameOffsetColumn>();
+	OutExtender.AddColumn<FNavigationToolTakeColumn>();
+	OutExtender.AddColumn<FNavigationToolCommentColumn>();
+
+	ISourceControlModule& SourceControlModule = ISourceControlModule::Get();
+	if (SourceControlModule.IsEnabled())
+	{
+		OutExtender.AddColumn<FNavigationToolRevisionControlColumn>();
+	}
+
+	FNavigationToolProvider::OnExtendColumns(OutExtender);
+}
+
+void FLevelSequenceNavigationToolProvider::OnExtendColumnViews(TSet<FNavigationToolColumnView>& OutColumnViews)
+{
+	using namespace UE::SequenceNavigator;
+
+	FNavigationToolColumnView& AnimationColumnView = OutColumnViews
+		.FindOrAdd(FNavigationToolColumnView(AnimationColumnViewName));
+
+	AnimationColumnView.VisibleColumns =
+		{
+			FNavigationToolPlayheadColumn::StaticColumnId(),
+			FNavigationToolColorColumn::StaticColumnId(),
+			FNavigationToolLabelColumn::StaticColumnId(),
+			FNavigationToolItemsColumn::StaticColumnId(),
+			FNavigationToolInTimeColumn::StaticColumnId(),
+			FNavigationToolOutTimeColumn::StaticColumnId(),
+			FNavigationToolHBiasColumn::StaticColumnId()
+		};
+
+	FNavigationToolProvider::OnExtendColumnViews(OutColumnViews);
+}
+
+void FLevelSequenceNavigationToolProvider::OnExtendItemChildren(UE::SequenceNavigator::INavigationTool& InTool
+	, const UE::SequenceNavigator::FNavigationToolItemPtr& InParentItem
+	, TArray<UE::SequenceNavigator::FNavigationToolItemPtr>& OutChildren
+	, const bool bInRecursive)
+{
+	using namespace UE::SequenceNavigator;
+
+	FNavigationToolProvider::OnExtendItemChildren(InTool, InParentItem, OutChildren, bInRecursive);
+
+	// Only extending root item
+	if (InParentItem->GetItemId() != FNavigationToolItemId::RootId)
+	{
+		return;
+	}
+
+	const TSharedPtr<ISequencer> Sequencer = InTool.GetSequencer();
+	if (!Sequencer.IsValid())
+	{
+		return;
+	}
+
+	UMovieSceneSequence* const RootSequence = Sequencer->GetRootMovieSceneSequence();
+	if (!RootSequence)
+	{
+		return;
+	}
+
+	const TSharedRef<FLevelSequenceNavigationToolProvider> Provider = SharedThis(this);
+
+	const FNavigationToolItemPtr NewItem = InTool.FindOrAdd<FNavigationToolSequence>(Provider, InParentItem, RootSequence, nullptr, 0);
+	const FNavigationToolItemFlagGuard Guard(NewItem, ENavigationToolItemFlags::IgnorePendingKill);
+	OutChildren.Add(NewItem);
+	if (bInRecursive)
+	{
+		NewItem->FindChildren(OutChildren, bInRecursive);
+	}
+}
+
+void FLevelSequenceNavigationToolProvider::OnExtendBuiltInFilters(TArray<FNavigationToolBuiltInFilterParams>& OutFilterParams)
+{
+	OutFilterParams.Add(FNavigationToolBuiltInFilterParams::CreateSequenceFilter());
+	OutFilterParams.Add(FNavigationToolBuiltInFilterParams::CreateTrackFilter());
+	OutFilterParams.Add(FNavigationToolBuiltInFilterParams::CreateBindingFilter());
+	OutFilterParams.Add(FNavigationToolBuiltInFilterParams::CreateMarkerFilter());
+
+	FNavigationToolProvider::OnExtendBuiltInFilters(OutFilterParams);
+}
+
+FMovieSceneEditorData* FLevelSequenceNavigationToolProvider::GetRootMovieSceneEditorData(const UE::SequenceNavigator::INavigationTool& InTool) const
+{
+	const TSharedPtr<ISequencer> Sequencer = InTool.GetSequencer();
+	if (!Sequencer.IsValid())
+	{
+		return nullptr;
+	}
+
+	UMovieSceneSequence* const RootMovieSceneSequence = Sequencer->GetRootMovieSceneSequence();
+	if (!RootMovieSceneSequence)
+	{
+		return nullptr;
+	}
+
+	UMovieScene* const RootMovieScene = RootMovieSceneSequence->GetMovieScene();
+	if (!RootMovieScene)
+	{
+		return nullptr;
+	}
+
+	return &RootMovieScene->GetEditorData();
+}
+
+FNavigationToolSaveState* FLevelSequenceNavigationToolProvider::GetSaveState(const UE::SequenceNavigator::INavigationTool& InTool) const
+{
+	if (FMovieSceneEditorData* const EditorData = GetRootMovieSceneEditorData(InTool))
+	{
+		return &EditorData->NavigationToolState;
+	}
+	return nullptr;
+}
+
+void FLevelSequenceNavigationToolProvider::SetSaveState(const UE::SequenceNavigator::INavigationTool& InTool
+	, const FNavigationToolSaveState& InSaveState) const
+{
+	if (FMovieSceneEditorData* const EditorData = GetRootMovieSceneEditorData(InTool))
+	{
+		EditorData->NavigationToolState = InSaveState;
+	}
+}
+
+#undef LOCTEXT_NAMESPACE
diff --git a/Engine/Plugins/Experimental/LevelSequenceNavigatorBridge/Source/LevelSequenceNavigatorBridge/Private/LevelSequenceNavigationToolProvider.h b/Engine/Plugins/Experimental/LevelSequenceNavigatorBridge/Source/LevelSequenceNavigatorBridge/Private/LevelSequenceNavigationToolProvider.h
new file mode 100644
--- /dev/null
+++ b/Engine/Plugins/Experimental/LevelSequenceNavigatorBridge/Source/LevelSequenceNavigatorBridge/Private/LevelSequenceNavigationToolProvider.h
@@ -0,0 +1,65 @@
+﻿// Copyright Epic Games, Inc. All Rights Reserved.
+
+#pragma once
+
+#include "NavigationToolDefines.h"
+#include "Providers/NavigationToolProvider.h"
+#include "Templates/SharedPointer.h"
+#include "UObject/WeakObjectPtr.h"
+
+class FUICommandList;
+class UMovieSceneSequence;
+struct FMovieSceneEditorData;
+struct FNavigationToolBuiltInFilterParams;
+struct FNavigationToolColumnView;
+struct FNavigationToolSaveState;
+
+namespace UE::SequenceNavigator
+{
+	class FNavigationToolColumnExtender;
+	class FNavigationToolItem;
+	class INavigationTool;
+}
+
+class FLevelSequenceNavigationToolProvider : public UE::SequenceNavigator::FNavigationToolProvider
+{
+public:
+	static const FName Identifier;
+
+	static const FText AnimationColumnViewName;
+
+	FLevelSequenceNavigationToolProvider();
+	virtual ~FLevelSequenceNavigationToolProvider() override = default;
+
+	//~ Begin INavigationToolProvider
+
+	virtual FName GetIdentifier() const override;
+
+	virtual TSet<TSubclassOf<UMovieSceneSequence>> GetSupportedSequenceClasses() const override;
+
+	virtual FText GetDefaultColumnView() const override;
+
+	virtual void BindCommands(const TSharedRef<FUICommandList>& InCommandList) override;
+
+	virtual void OnActivate() override;
+	virtual void OnDeactivate() override;
+
+	virtual void OnExtendColumns(UE::SequenceNavigator::FNavigationToolColumnExtender& OutExtender) override;
+	virtual void OnExtendColumnViews(TSet<FNavigationToolColumnView>& OutColumnViews) override;
+	virtual void OnExtendItemChildren(UE::SequenceNavigator::INavigationTool& InTool
+		, const UE::SequenceNavigator::FNavigationToolItemPtr& InParentItem
+		, TArray<UE::SequenceNavigator::FNavigationToolItemPtr>& OutChildren
+		, const bool bInRecursive) override;
+	virtual void OnExtendBuiltInFilters(TArray<FNavigationToolBuiltInFilterParams>& OutFilterParams) override;
+
+	virtual FNavigationToolSaveState* GetSaveState(const UE::SequenceNavigator::INavigationTool& InTool) const override;
+	virtual void SetSaveState(const UE::SequenceNavigator::INavigationTool& InTool
+		, const FNavigationToolSaveState& InSaveState) const override;
+
+	//~ End INavigationToolProvider
+
+private:
+	FMovieSceneEditorData* GetRootMovieSceneEditorData(const UE::SequenceNavigator::INavigationTool& InTool) const;
+
+	TSharedRef<FUICommandList> ToolCommands;
+};
diff --git a/Engine/Plugins/Experimental/LevelSequenceNavigatorBridge/Source/LevelSequenceNavigatorBridge/Private/LevelSequenceNavigatorBridgeModule.cpp b/Engine/Plugins/Experimental/LevelSequenceNavigatorBridge/Source/LevelSequenceNavigatorBridge/Private/LevelSequenceNavigatorBridgeModule.cpp
new file mode 100644
--- /dev/null
+++ b/Engine/Plugins/Experimental/LevelSequenceNavigatorBridge/Source/LevelSequenceNavigatorBridge/Private/LevelSequenceNavigatorBridgeModule.cpp
@@ -0,0 +1,71 @@
+﻿// Copyright Epic Games, Inc. All Rights Reserved.
+
+#include "LevelSequenceNavigatorBridgeModule.h"
+#include "ISequencer.h"
+#include "ISequencerModule.h"
+#include "LevelSequenceNavigationToolProvider.h"
+#include "NavigationToolExtender.h"
+
+#define LOCTEXT_NAMESPACE "LevelSequenceNavigatorBridgeModule"
+
+using namespace UE::SequenceNavigator;
+
+void FLevelSequenceNavigatorBridgeModule::StartupModule()
+{
+	if (ISequencerModule* const SequencerModule = FModuleManager::LoadModulePtr<ISequencerModule>(TEXT("Sequencer")))
+	{
+		SequencerCreatedHandle = SequencerModule->RegisterOnSequencerCreated(
+			FOnSequencerCreated::FDelegate::CreateRaw(this, &FLevelSequenceNavigatorBridgeModule::OnSequencerCreated));
+	}
+}
+
+void FLevelSequenceNavigatorBridgeModule::ShutdownModule()
+{
+	if (ISequencerModule* const SequencerModule = FModuleManager::LoadModulePtr<ISequencerModule>(TEXT("Sequencer")))
+	{
+		SequencerModule->UnregisterOnSequencerCreated(SequencerCreatedHandle);
+		SequencerCreatedHandle.Reset();
+	}
+}
+
+void FLevelSequenceNavigatorBridgeModule::OnSequencerCreated(const TSharedRef<ISequencer> InSequencer)
+{
+	const FName ToolId = FNavigationToolExtender::GetToolInstanceId(*InSequencer);
+
+	const TSharedPtr<FNavigationToolProvider> FoundProvider = FNavigationToolExtender::FindToolProvider(ToolId
+		, FLevelSequenceNavigationToolProvider::Identifier);
+	if (!FoundProvider.IsValid())
+	{
+		NavigationToolProvider = MakeShared<FLevelSequenceNavigationToolProvider>();
+		if (NavigationToolProvider->IsSequenceSupported(InSequencer->GetRootMovieSceneSequence()))
+		{
+			FNavigationToolExtender::RegisterToolProvider(InSequencer, NavigationToolProvider.ToSharedRef());
+		}
+		else
+		{
+			NavigationToolProvider.Reset();
+		}
+	}
+
+	SequencerClosedHandle = InSequencer->OnCloseEvent().AddRaw(this, &FLevelSequenceNavigatorBridgeModule::OnSequencerClosed);
+}
+
+void FLevelSequenceNavigatorBridgeModule::OnSequencerClosed(const TSharedRef<ISequencer> InSequencer)
+{
+	if (!NavigationToolProvider.IsValid())
+	{
+		return;
+	}
+
+	const FName ToolId = FNavigationToolExtender::GetToolInstanceId(*InSequencer);
+
+	if (FNavigationToolExtender::UnregisterToolProvider(ToolId, NavigationToolProvider->GetIdentifier()))
+	{
+		NavigationToolProvider.Reset();
+		SequencerClosedHandle.Reset();
+	}
+}
+
+#undef LOCTEXT_NAMESPACE
+	
+IMPLEMENT_MODULE(FLevelSequenceNavigatorBridgeModule, LevelSequenceNavigatorBridge)
diff --git a/Engine/Plugins/Experimental/LevelSequenceNavigatorBridge/Source/LevelSequenceNavigatorBridge/Private/LevelSequenceNavigatorBridgeModule.h b/Engine/Plugins/Experimental/LevelSequenceNavigatorBridge/Source/LevelSequenceNavigatorBridge/Private/LevelSequenceNavigatorBridgeModule.h
new file mode 100644
--- /dev/null
+++ b/Engine/Plugins/Experimental/LevelSequenceNavigatorBridge/Source/LevelSequenceNavigatorBridge/Private/LevelSequenceNavigatorBridgeModule.h
@@ -0,0 +1,29 @@
+﻿// Copyright Epic Games, Inc. All Rights Reserved.
+
+#pragma once
+
+#include "Delegates/IDelegateInstance.h"
+#include "Modules/ModuleManager.h"
+#include "Templates/SharedPointer.h"
+
+class FLevelSequenceNavigationToolProvider;
+class ISequencer;
+
+class FLevelSequenceNavigatorBridgeModule : public IModuleInterface
+{
+public:
+	//~ Begin IModuleInterface
+	virtual void StartupModule() override;
+	virtual void ShutdownModule() override;
+	//~ End IModuleInterface
+
+private:
+	void OnSequencerCreated(const TSharedRef<ISequencer> InSequencer);
+	void OnSequencerClosed(const TSharedRef<ISequencer> InSequencer);
+
+	FDelegateHandle SequencerCreatedHandle;
+	FDelegateHandle SequencerClosedHandle;
+
+	/** The provider that supplies data and extends the Navigation Tool */
+	TSharedPtr<FLevelSequenceNavigationToolProvider> NavigationToolProvider;
+};
diff --git a/Engine/Plugins/Experimental/SequenceNavigator/SequenceNavigator.uplugin b/Engine/Plugins/Experimental/SequenceNavigator/SequenceNavigator.uplugin
--- a/Engine/Plugins/Experimental/SequenceNavigator/SequenceNavigator.uplugin
+++ b/Engine/Plugins/Experimental/SequenceNavigator/SequenceNavigator.uplugin
@@ -13,6 +13,7 @@
 	"IsBetaVersion": false,
 	"IsExperimentalVersion": true,
 	"Installed": false,
+	"EnabledByDefault": true,
 	"Modules": [
 		{
 			"Name": "SequenceNavigator",
diff --git a/Engine/Plugins/Experimental/SequenceNavigator/Source/SequenceNavigator/Public/Providers/NavigationToolProvider.h b/Engine/Plugins/Experimental/SequenceNavigator/Source/SequenceNavigator/Public/Providers/NavigationToolProvider.h
--- a/Engine/Plugins/Experimental/SequenceNavigator/Source/SequenceNavigator/Public/Providers/NavigationToolProvider.h
+++ b/Engine/Plugins/Experimental/SequenceNavigator/Source/SequenceNavigator/Public/Providers/NavigationToolProvider.h
@@ -41,7 +41,7 @@ public:
 	/** Ensure saved editor only data contains an entry for a specific view, creating it if necessary */
 	void EnsureToolViewCount(const INavigationTool& InTool, const int32 InToolViewId);
 
-	bool IsSequenceSupported(UMovieSceneSequence* const InSequence) const;
+	SEQUENCENAVIGATOR_API bool IsSequenceSupported(UMovieSceneSequence* const InSequence) const;
 
 private:
 	friend class FNavigationToolExtender;
diff --git a/Engine/Plugins/MovieScene/LevelSequenceEditor/LevelSequenceEditor.uplugin b/Engine/Plugins/MovieScene/LevelSequenceEditor/LevelSequenceEditor.uplugin
--- a/Engine/Plugins/MovieScene/LevelSequenceEditor/LevelSequenceEditor.uplugin
+++ b/Engine/Plugins/MovieScene/LevelSequenceEditor/LevelSequenceEditor.uplugin
@@ -28,10 +28,6 @@
 		{
 			"Name": "SequencerScripting",
 			"Enabled": true
-		},
-		{
-			"Name": "SequenceNavigator",
-			"Enabled": true
 		}
 	]
 }
\ No newline at end of file
diff --git a/Engine/Plugins/MovieScene/LevelSequenceEditor/Source/LevelSequenceEditor/LevelSequenceEditor.Build.cs b/Engine/Plugins/MovieScene/LevelSequenceEditor/Source/LevelSequenceEditor/LevelSequenceEditor.Build.cs
--- a/Engine/Plugins/MovieScene/LevelSequenceEditor/Source/LevelSequenceEditor/LevelSequenceEditor.Build.cs
+++ b/Engine/Plugins/MovieScene/LevelSequenceEditor/Source/LevelSequenceEditor/LevelSequenceEditor.Build.cs
@@ -50,7 +50,6 @@ public class LevelSequenceEditor : ModuleRules
 				"MovieSceneTracks",
 				"Projects",
                 "PropertyEditor",
-                "SequenceNavigator",
 				"SequencerCore",
 				"Sequencer",
                 "Slate",
diff --git a/Engine/Plugins/MovieScene/LevelSequenceEditor/Source/LevelSequenceEditor/Private/LevelSequenceEditorToolkit.cpp b/Engine/Plugins/MovieScene/LevelSequenceEditor/Source/LevelSequenceEditor/Private/LevelSequenceEditorToolkit.cpp
--- a/Engine/Plugins/MovieScene/LevelSequenceEditor/Source/LevelSequenceEditor/Private/LevelSequenceEditorToolkit.cpp
+++ b/Engine/Plugins/MovieScene/LevelSequenceEditor/Source/LevelSequenceEditor/Private/LevelSequenceEditorToolkit.cpp
@@ -44,8 +44,6 @@
 #include "ClassViewerModule.h"
 #include "SequencerUtilities.h"
 #include "ViewportToolbar/UnrealEdViewportToolbarContext.h"
-#include "LevelSequenceNavigationToolProvider.h"
-#include "NavigationToolExtender.h"
 
 // @todo sequencer: hack: setting defaults for transform tracks
 
@@ -317,21 +315,6 @@ void FLevelSequenceEditorToolkit::Initialize(const EToolkitMode::Type Mode, cons
 			SelectionSection.AddEntry(Entry);
 		}
 	}
-
-	if (Sequencer.IsValid())
-	{
-		using namespace UE::SequenceNavigator;
-
-		const FName ToolId = FNavigationToolExtender::GetToolInstanceId(*Sequencer);
-
-		const TSharedPtr<FNavigationToolProvider> FoundProvider = FNavigationToolExtender::FindToolProvider(ToolId
-			, FLevelSequenceNavigationToolProvider::Identifier);
-		if (!FoundProvider.IsValid())
-		{
-			NavigationToolProvider = MakeShared<FLevelSequenceNavigationToolProvider>();
-			FNavigationToolExtender::RegisterToolProvider(Sequencer.ToSharedRef(), NavigationToolProvider.ToSharedRef());
-		}
-	}
 }
 
 
@@ -997,16 +980,6 @@ void FLevelSequenceEditorToolkit::SaveAsset_Execute()
 
 void FLevelSequenceEditorToolkit::OnClose()
 {
-	if (Sequencer.IsValid() && NavigationToolProvider.IsValid())
-	{
-		using namespace UE::SequenceNavigator;
-
-		const FName ToolId = FNavigationToolExtender::GetToolInstanceId(*Sequencer);
-		FNavigationToolExtender::UnregisterToolProvider(ToolId, NavigationToolProvider->GetIdentifier());
-
-		NavigationToolProvider.Reset();
-	}
-
 	// Remove the viewport toolbar extensions we added earlier.
 	UToolMenus::Get()->UnregisterOwnerByName(UE::MovieScene::Private::ViewportToolbarOwnerName);
 
diff --git a/Engine/Plugins/MovieScene/LevelSequenceEditor/Source/LevelSequenceEditor/Private/LevelSequenceEditorToolkit.h b/Engine/Plugins/MovieScene/LevelSequenceEditor/Source/LevelSequenceEditor/Private/LevelSequenceEditorToolkit.h
--- a/Engine/Plugins/MovieScene/LevelSequenceEditor/Source/LevelSequenceEditor/Private/LevelSequenceEditorToolkit.h
+++ b/Engine/Plugins/MovieScene/LevelSequenceEditor/Source/LevelSequenceEditor/Private/LevelSequenceEditorToolkit.h
@@ -9,7 +9,6 @@ class FSpawnTabArgs;
 struct FFrameNumber;
 
 class AActor;
-class FLevelSequenceNavigationToolProvider;
 class FMenuBuilder;
 class FToolBarBuilder;
 class IAssetViewport;
@@ -148,9 +147,6 @@ private:
 	/** The sequencer used by this editor. */
 	TSharedPtr<ISequencer> Sequencer;
 
-	/** The provider that supplies data and extends the Navigation Tool */
-	TSharedPtr<FLevelSequenceNavigationToolProvider> NavigationToolProvider;
-
 	/** Pointer to the style set to use for toolkits. */
 	TSharedRef<ISlateStyle> Style;
 
diff --git a/Engine/Plugins/MovieScene/LevelSequenceEditor/Source/LevelSequenceEditor/Private/LevelSequenceNavigationToolProvider.cpp b/Engine/Plugins/MovieScene/LevelSequenceEditor/Source/LevelSequenceEditor/Private/LevelSequenceNavigationToolProvider.cpp
deleted file mode 100644
diff --git a/Engine/Plugins/MovieScene/LevelSequenceEditor/Source/LevelSequenceEditor/Private/LevelSequenceNavigationToolProvider.h b/Engine/Plugins/MovieScene/LevelSequenceEditor/Source/LevelSequenceEditor/Private/LevelSequenceNavigationToolProvider.h
deleted file mode 100644
