[SequenceNavigator] Change to support UMovieSceneSequence instead of ULevelSequence

diff --git a/Engine/Plugins/Experimental/LevelSequenceNavigatorBridge/Source/LevelSequenceNavigatorBridge/Private/LevelSequenceNavigationToolProvider.cpp b/Engine/Plugins/Experimental/LevelSequenceNavigatorBridge/Source/LevelSequenceNavigatorBridge/Private/LevelSequenceNavigationToolProvider.cpp
--- a/Engine/Plugins/Experimental/LevelSequenceNavigatorBridge/Source/LevelSequenceNavigatorBridge/Private/LevelSequenceNavigationToolProvider.cpp
+++ b/Engine/Plugins/Experimental/LevelSequenceNavigatorBridge/Source/LevelSequenceNavigatorBridge/Private/LevelSequenceNavigationToolProvider.cpp
@@ -22,7 +22,6 @@
 #include "ISequencer.h"
 #include "ISourceControlModule.h"
 #include "Items/NavigationToolSequence.h"
-#include "LevelSequence.h"
 #include "MovieScene.h"
 #include "MovieSceneSequence.h"
 #include "NavigationToolSettings.h"
@@ -45,7 +44,7 @@ FName FLevelSequenceNavigationToolProvider::GetIdentifier() const
 
 TSet<TSubclassOf<UMovieSceneSequence>> FLevelSequenceNavigationToolProvider::GetSupportedSequenceClasses() const
 {
-	return { ULevelSequence::StaticClass() };
+	return { UMovieSceneSequence::StaticClass() };
 }
 
 FText FLevelSequenceNavigationToolProvider::GetDefaultColumnView() const
diff --git a/Engine/Plugins/Experimental/SequenceNavigator/Source/SequenceNavigator/Private/Providers/NavigationToolProvider.cpp b/Engine/Plugins/Experimental/SequenceNavigator/Source/SequenceNavigator/Private/Providers/NavigationToolProvider.cpp
--- a/Engine/Plugins/Experimental/SequenceNavigator/Source/SequenceNavigator/Private/Providers/NavigationToolProvider.cpp
+++ b/Engine/Plugins/Experimental/SequenceNavigator/Source/SequenceNavigator/Private/Providers/NavigationToolProvider.cpp
@@ -190,10 +190,25 @@ void FNavigationToolProvider::CleanupExtendedColumnViews()
 
 bool FNavigationToolProvider::IsSequenceSupported(UMovieSceneSequence* const InSequence) const
 {
-	if (InSequence)
+	if (!InSequence)
 	{
-		return GetSupportedSequenceClasses().Contains(InSequence->GetClass());
+		return false;
 	}
+
+	UClass* const SequenceClass = InSequence->GetClass();
+	if (!SequenceClass)
+	{
+		return false;
+	}
+
+	for (const TSubclassOf<UMovieSceneSequence>& SupportedClass : GetSupportedSequenceClasses())
+	{
+		if (SequenceClass->IsChildOf(SupportedClass))
+		{
+			return true;
+		}
+	}
+
 	return false;
 }
 
